[GLTFExporter] chore: add user-oriented readme inside plugin with implementation status

commit 3a23509f7785f6c3eba01fa6136b84c87499fbd6
Date: Wed, 20 May 2020 03:10:11 +0200

    chore: add user-oriented readme inside plugin with implementation status
    (cherry picked from commit ca3620ac9bd0d15d4e75734c381f4fe6b93e258b​)

#skipci

diff --git a/Engine/Plugins/Enterprise/GLTFExporter/MigratedFromGit/README.md b/Engine/Plugins/Enterprise/GLTFExporter/MigratedFromGit/README.md
new file mode 100644
--- /dev/null
+++ b/Engine/Plugins/Enterprise/GLTFExporter/MigratedFromGit/README.md
@@ -0,0 +1,43 @@
+glTF exporter plugin
+====================
+
+An Unreal Editor plugin for exporting assets to web-optimized glTF 2.0 (incl custom and ratified extensions).
+
+
+Implementation status
+---------------------
+
+Feature                   | Status |
+------------------------- | :----: |
+Web-based Viewer          |   👷   |
+Models                    |   👷   |
+└ Static Mesh             |   ✔️   |
+└ Skeletal Mesh           |        |
+└ Animation Sequence      |        |
+Materials                 |   👷   |
+└ Default Lit             |   ✔️   |
+└ Unlit                   |   👷   |
+└ Clear Coat              |        |
+└ Blend Modes             |   ✔️   |
+└ Constant Input          |   ✔️   |
+└ Parameter Input         |   ✔️   |
+└ Texture Input           |   👷   |
+└ Material Instance       |   ✔️   |
+└ Material Baking         |        |
+Levels                    |   👷   |
+└ Actor/Blueprint         |   ✔️   |
+└ Static Mesh             |   ✔️   |
+└ Skeletal Mesh           |        |
+└ Light                   |        |
+└ Camera                  |        |
+└ Orbital Camera          |        |
+└ Reflection Capture      |        |
+└ HDRI Backdrop           |   👷   |
+└ Interaction Hotspot     |   👷   |
+└ Level Variant Sets      |        |
+└ Baked Lightmaps         |        |
+Textures                  |        |
+└ 2D                      |        |
+└ Cubemap                 |        |
+└ Lightmap                |        |
+Exporter Options          |   👷   |
