[GLTFExporter] chore: add installation for plugin package

commit 414da74e28b4119f597c6c1855f5f0eb7f01fae9
Date: Wed, 20 May 2020 03:12:34 +0200

    chore: add installation for plugin package
    (cherry picked from commit 60e69572731ec84ba41ca9be0061a00c13a8c22c​)

#skipci

diff --git a/Engine/Plugins/Enterprise/GLTFExporter/MigratedFromGit/README.md b/Engine/Plugins/Enterprise/GLTFExporter/MigratedFromGit/README.md
--- a/Engine/Plugins/Enterprise/GLTFExporter/MigratedFromGit/README.md
+++ b/Engine/Plugins/Enterprise/GLTFExporter/MigratedFromGit/README.md
@@ -41,3 +41,18 @@ Textures                  |        |
 └ Cubemap                 |        |
 └ Lightmap                |        |
 Exporter Options          |   👷   |
+
+
+Installation
+------------
+
+* Alt 1: Install to project:
+    1. Go to project folder which contains `[ProjectName].uproject`
+    1. Create a folder called `Plugins` (if it doesn't already exist).
+    1. Copy `GLTFExporter` folder into the `Plugins` folder.
+
+* Alt 2: Install to Unreal Engine:
+    1. Go to the plugin folder of Unreal Engine which is `Engine/Plugins`
+    1. Copy `GLTFExporter` folder into the `Plugins` folder.
+
+In Unreal Editor, open `Menu -> Edit -> Plugins` and make sure `glTF Exporter` is installed and enabled.
