[GLTFExporter] chore: update implementation status

commit 868dd707feae23c6212b54fe7a8fbd93771ced18
Date: Wed, 20 May 2020 11:11:05 +0200

    chore: update implementation status
    (cherry picked from commit 9a59a4c373c38f21b6d5571f975614dd10562065​)

#skipci

diff --git a/Engine/Plugins/Enterprise/GLTFExporter/MigratedFromGit/README.md b/Engine/Plugins/Enterprise/GLTFExporter/MigratedFromGit/README.md
--- a/Engine/Plugins/Enterprise/GLTFExporter/MigratedFromGit/README.md
+++ b/Engine/Plugins/Enterprise/GLTFExporter/MigratedFromGit/README.md
@@ -14,11 +14,12 @@ Models                    |   👷   |
 └ Static Mesh             |   ✔️   |
 └ Skeletal Mesh           |        |
 └ Animation Sequence      |        |
+└ Mesh Quantization       |   👷   |
 Materials                 |   👷   |
 └ Default Lit             |   ✔️   |
 └ Unlit                   |   👷   |
 └ Clear Coat              |        |
-└ Blend Modes             |   ✔️   |
+└ Blend Modes             |   👷   |
 └ Constant Input          |   ✔️   |
 └ Parameter Input         |   ✔️   |
 └ Texture Input           |   👷   |
@@ -38,6 +39,7 @@ Levels                    |   👷   |
 └ Baked Lightmaps         |        |
 Textures                  |        |
 └ 2D                      |        |
+└ HDR                     |        |
 └ Cubemap                 |        |
 └ Lightmap                |        |
 Exporter Options          |   👷   |
