[GLTFExporter] chore: add usage instructions to plugin readme

commit 3783b3f7107897346cfef9fbf4b90fa32e06498b
Date: Thu, 21 May 2020 03:29:44 +0200

    chore: add usage instructions to plugin readme
    (cherry picked from commit 29051b8be466ce61e32640a27934b36995ff16fd​)

#skipci

diff --git a/Engine/Plugins/Enterprise/GLTFExporter/MigratedFromGit/README.md b/Engine/Plugins/Enterprise/GLTFExporter/MigratedFromGit/README.md
--- a/Engine/Plugins/Enterprise/GLTFExporter/MigratedFromGit/README.md
+++ b/Engine/Plugins/Enterprise/GLTFExporter/MigratedFromGit/README.md
@@ -1,7 +1,7 @@
 glTF exporter plugin
 ====================
 
-An Unreal Editor plugin for exporting assets to web-optimized glTF 2.0 (incl custom and ratified extensions).
+An Unreal Editor plugin for exporting to web-optimized glTF 2.0 (incl custom and ratified extensions).
 
 
 Implementation status
@@ -49,12 +49,30 @@ Installation
 ------------
 
 * Alt 1: Install to project:
-    1. Go to project folder which contains `[ProjectName].uproject`
-    1. Create a folder called `Plugins` (if it doesn't already exist).
-    1. Copy `GLTFExporter` folder into the `Plugins` folder.
+  1. Go to project folder which contains `[ProjectName].uproject`
+  1. Create a folder called `Plugins` (if it doesn't already exist).
+  1. Copy `GLTFExporter` folder into the `Plugins` folder.
 
 * Alt 2: Install to Unreal Engine:
-    1. Go to the plugin folder of Unreal Engine which is `Engine/Plugins`
-    1. Copy `GLTFExporter` folder into the `Plugins` folder.
+  1. Go to the plugin folder of Unreal Engine which is `Engine/Plugins`
+  1. Copy `GLTFExporter` folder into the `Plugins` folder.
 
 In Unreal Editor, open `Menu -> Edit -> Plugins` and make sure `glTF Exporter` is installed and enabled.
+
+
+Usage
+-----
+
+* Alt 1: Export asset via Content Browser
+  1. Right-click on a `StaticMesh`, `Level`, or `Material` asset in the Content Browser.
+  1. Select `Asset Actions -> Export...`
+  1. Change `Save as type` to `.gltf` and click `Save`
+  1. When `Export Options` window is shown, click `Export`
+
+* Alt 2: Export current level via File Menu
+  1. Select any number of actors in the current level
+  1. In the top menu, select `File -> Export Selected...`
+  1. Change `Save as type` to `.gltf` and click `Save`
+  1. When `Export Options` window is shown, click `Export`
+
+The exported file can be opened in any glTF viewer including those readily available on the web.
