[GLTFExporter] chore: update implementation status in plugin readme

commit 67318b19a756b93b509f38b4743eb13b9450538c
Date: Tue, 21 Jul 2020 18:40:25 +0200

    chore: update implementation status in plugin readme
    (cherry picked from commit 2097aba63a513dde4ba0b3cb8fe924581c7a8cd9​)

#skipci

diff --git a/Engine/Plugins/Enterprise/GLTFExporter/MigratedFromGit/README.md b/Engine/Plugins/Enterprise/GLTFExporter/MigratedFromGit/README.md
--- a/Engine/Plugins/Enterprise/GLTFExporter/MigratedFromGit/README.md
+++ b/Engine/Plugins/Enterprise/GLTFExporter/MigratedFromGit/README.md
@@ -12,37 +12,40 @@ Feature                   | Status |
 Web-based Viewer          |   👷   |
 Models                    |   👷   |
 └ Static Mesh             |   ✔️   |
-└ Skeletal Mesh           |        |
-└ Animation Sequence      |        |
+└ Skeletal Mesh           |   ✔️   |
+└ Animation Sequence      |   👷   |
 └ Mesh Quantization       |   👷   |
-Materials                 |   👷   |
+Materials                 |   ✔️   |
 └ Default Lit             |   ✔️   |
-└ Unlit                   |   👷   |
-└ Clear Coat              |        |
-└ Blend Modes             |   👷   |
+└ Unlit                   |   ✔️   |
+└ Clear Coat              |   ✔️   |
+└ Blend Modes             |   ✔️   |
 └ Constant Input          |   ✔️   |
 └ Parameter Input         |   ✔️   |
-└ Texture Input           |   👷   |
+└ Texture Input           |   ✔️   |
 └ Material Instance       |   ✔️   |
-└ Material Baking         |        |
+└ Material Baking         |   ✔️   |
 Levels                    |   👷   |
 └ Actor/Blueprint         |   ✔️   |
 └ Static Mesh             |   ✔️   |
-└ Skeletal Mesh           |        |
-└ Light                   |        |
-└ Camera                  |        |
+└ Skeletal Mesh           |   ✔️   |
+└ Light                   |   ✔️   |
+└ Camera                  |   ✔️   |
 └ Orbital Camera          |        |
 └ Reflection Capture      |        |
-└ HDRI Backdrop           |   👷   |
+└ HDRI Backdrop           |   ✔️   |
 └ Interaction Hotspot     |   👷   |
-└ Level Variant Sets      |        |
-└ Baked Lightmaps         |        |
-Textures                  |        |
-└ 2D                      |        |
-└ HDR                     |        |
-└ Cubemap                 |        |
-└ Lightmap                |        |
-Exporter Options          |   👷   |
+└ Level Variant Sets      |   ✔️   |
+└ Baked Lightmaps         |   ✔️   |
+Textures                  |   ✔️   |
+└ 2D                      |   ✔️   |
+└ HDR                     |   ✔️   |
+└ Cubemap                 |   ✔️   |
+└ Lightmap                |   ✔️   |
+Exporter                  |   👷   |
+└ Options                 |   👷   |
+└ Logging                 |   👷   |
+└ Progress Bar            |   👷   |
 
 
 Installation
