[GLTFExporter] chore: add more detailed implementation status for level variant sets

commit 70ccecaa381e4cd76c349ff49b9bd1889cfac403
Date: Thu, 23 Jul 2020 22:36:53 +0200

    chore: add more detailed implementation status for level variant sets
    (cherry picked from commit 8089f8e89311f4a0b8a58f47e7d1a2909829523e​)

#skipci

diff --git a/Engine/Plugins/Enterprise/GLTFExporter/MigratedFromGit/README.md b/Engine/Plugins/Enterprise/GLTFExporter/MigratedFromGit/README.md
--- a/Engine/Plugins/Enterprise/GLTFExporter/MigratedFromGit/README.md
+++ b/Engine/Plugins/Enterprise/GLTFExporter/MigratedFromGit/README.md
@@ -42,6 +42,10 @@ Textures                  |   ✔️   |
 └ HDR                     |   ✔️   |
 └ Cubemap                 |   ✔️   |
 └ Lightmap                |   ✔️   |
+Level Variant Sets        |   👷   |
+└ Visible                 |   ✔️   |
+└ Material                |   ✔️   |
+└ Mesh                    |        |
 Exporter                  |   👷   |
 └ Options                 |   👷   |
 └ Logging                 |   👷   |
