[GLTFExporter] chore: update implementation status in plugin readme

commit efad9ce87b5e0ad98ce27c57f9822efb8a8d422c
Date: Tue, 6 Oct 2020 18:21:39 +0200

    chore: update implementation status in plugin readme
    (cherry picked from commit 8b4784213a2a47ad31eb6e8aab1047982671ac64​)

#skipci

diff --git a/Engine/Plugins/Enterprise/GLTFExporter/MigratedFromGit/README.md b/Engine/Plugins/Enterprise/GLTFExporter/MigratedFromGit/README.md
--- a/Engine/Plugins/Enterprise/GLTFExporter/MigratedFromGit/README.md
+++ b/Engine/Plugins/Enterprise/GLTFExporter/MigratedFromGit/README.md
@@ -10,11 +10,13 @@ Implementation status
 Feature                   | Status |
 ------------------------- | :----: |
 Web-based Viewer          |   👷   |
-Models                    |   👷   |
+└ UX/UI                   |   👷   |
+└ Offline                 |   👷   |
+Models                    |   ✔️   |
 └ Static Mesh             |   ✔️   |
 └ Skeletal Mesh           |   ✔️   |
-└ Animation Sequence      |   👷   |
-└ Mesh Quantization       |   👷   |
+└ Animation Sequence      |   ✔️   |
+└ Mesh Quantization       |   ✔️   |
 Materials                 |   ✔️   |
 └ Default Lit             |   ✔️   |
 └ Unlit                   |   ✔️   |
@@ -31,25 +33,28 @@ Levels                    |   👷   |
 └ Skeletal Mesh           |   ✔️   |
 └ Light                   |   ✔️   |
 └ Camera                  |   ✔️   |
-└ Orbital Camera          |        |
-└ Reflection Capture      |        |
+└ Orbital Camera          |   ✔️   |
+└ Reflection Capture      |   👷   |
 └ HDRI Backdrop           |   ✔️   |
-└ Interaction Hotspot     |   👷   |
+└ Animation Hotspot       |   ✔️   |
 └ Level Variant Sets      |   ✔️   |
 └ Baked Lightmaps         |   ✔️   |
+└ Level Sequence          |   ✔️   |
 Textures                  |   ✔️   |
 └ 2D                      |   ✔️   |
 └ HDR                     |   ✔️   |
 └ Cubemap                 |   ✔️   |
 └ Lightmap                |   ✔️   |
-Level Variant Sets        |   👷   |
+Level Variant Sets        |   ✔️   |
 └ Visible                 |   ✔️   |
 └ Material                |   ✔️   |
-└ Mesh                    |        |
+└ Mesh                    |   ✔️   |
+Level Sequence            |   ✔️   |
+└ Transform               |   ✔️   |
 Exporter                  |   👷   |
 └ Options                 |   👷   |
 └ Logging                 |   👷   |
-└ Progress Bar            |   👷   |
+└ Progress Bars           |   ✔️   |
 
 
 Installation
