[GLTFExporter] chore: update plugin README

commit f8f0282f1abfa8260c04b7526173daa8bf7d117a
Date: Mon, 16 Nov 2020 19:43:00 +0100

    chore: update plugin README
    (cherry picked from commit 0f0d83e2c57be0739bb4e74814a78db9a9f81410​)

#skipci

diff --git a/Engine/Plugins/Enterprise/GLTFExporter/MigratedFromGit/README.md b/Engine/Plugins/Enterprise/GLTFExporter/MigratedFromGit/README.md
--- a/Engine/Plugins/Enterprise/GLTFExporter/MigratedFromGit/README.md
+++ b/Engine/Plugins/Enterprise/GLTFExporter/MigratedFromGit/README.md
@@ -9,9 +9,9 @@ Implementation status
 
 Feature                   | Status |
 ------------------------- | :----: |
-Web-based Viewer          |   👷   |
-└ UX/UI                   |   👷   |
-└ Offline                 |   👷   |
+Web-based Viewer          |   ✔️   |
+└ UX/UI                   |   ✔️   |
+└ Launcher                |   ✔️   |
 Models                    |   ✔️   |
 └ Static Mesh             |   ✔️   |
 └ Skeletal Mesh           |   ✔️   |
@@ -33,7 +33,7 @@ Levels                    |   👷   |
 └ Skeletal Mesh           |   ✔️   |
 └ Light                   |   ✔️   |
 └ Camera                  |   ✔️   |
-└ Orbital Camera          |   ✔️   |
+└ Camera Control          |   ✔️   |
 └ Reflection Capture      |   👷   |
 └ HDRI Backdrop           |   ✔️   |
 └ Animation Hotspot       |   ✔️   |
