[GLTFExporter] chore: update README files to more closely match those at github.com/ue4plugins

commit 7e65a45ca942980a6a94449e7bd353a3ca361382
Date: Sun, 29 Nov 2020 22:11:28 +0100

    chore: update README files to more closely match those at github.com/ue4plugins
    (cherry picked from commit 71f08269f27ac6802be479711452c281ba4361e7​)

#skipci

diff --git a/Engine/Plugins/Enterprise/GLTFExporter/MigratedFromGit/README.md b/Engine/Plugins/Enterprise/GLTFExporter/MigratedFromGit/README.md
--- a/Engine/Plugins/Enterprise/GLTFExporter/MigratedFromGit/README.md
+++ b/Engine/Plugins/Enterprise/GLTFExporter/MigratedFromGit/README.md
@@ -1,60 +1,7 @@
 glTF exporter plugin
 ====================
 
-An Unreal Editor plugin for exporting to web-optimized glTF 2.0 (incl custom and ratified extensions).
-
-
-Implementation status
----------------------
-
-Feature                   | Status |
-------------------------- | :----: |
-Web-based Viewer          |   ✔️   |
-└ UX/UI                   |   ✔️   |
-└ Launcher                |   ✔️   |
-Models                    |   ✔️   |
-└ Static Mesh             |   ✔️   |
-└ Skeletal Mesh           |   ✔️   |
-└ Animation Sequence      |   ✔️   |
-└ Mesh Quantization       |   ✔️   |
-Materials                 |   ✔️   |
-└ Default Lit             |   ✔️   |
-└ Unlit                   |   ✔️   |
-└ Clear Coat              |   ✔️   |
-└ Blend Modes             |   ✔️   |
-└ Constant Input          |   ✔️   |
-└ Parameter Input         |   ✔️   |
-└ Texture Input           |   ✔️   |
-└ Material Instance       |   ✔️   |
-└ Material Baking         |   ✔️   |
-Levels                    |   👷   |
-└ Actor/Blueprint         |   ✔️   |
-└ Static Mesh             |   ✔️   |
-└ Skeletal Mesh           |   ✔️   |
-└ Light                   |   ✔️   |
-└ Camera                  |   ✔️   |
-└ Camera Control          |   ✔️   |
-└ Reflection Capture      |   👷   |
-└ HDRI Backdrop           |   ✔️   |
-└ Animation Hotspot       |   ✔️   |
-└ Level Variant Sets      |   ✔️   |
-└ Baked Lightmaps         |   ✔️   |
-└ Level Sequence          |   ✔️   |
-Textures                  |   ✔️   |
-└ 2D                      |   ✔️   |
-└ HDR                     |   ✔️   |
-└ Cubemap                 |   ✔️   |
-└ Lightmap                |   ✔️   |
-Level Variant Sets        |   ✔️   |
-└ Visible                 |   ✔️   |
-└ Material                |   ✔️   |
-└ Mesh                    |   ✔️   |
-Level Sequence            |   ✔️   |
-└ Transform               |   ✔️   |
-Exporter                  |   👷   |
-└ Options                 |   👷   |
-└ Logging                 |   👷   |
-└ Progress Bars           |   ✔️   |
+Unreal Editor plugin for exporting UE4 assets to glTF 2.0 (with custom extensions).
 
 
 Installation
@@ -76,7 +23,7 @@ Usage
 -----
 
 * Alt 1: Export asset via Content Browser
-  1. Right-click on a `StaticMesh`, `Level`, or `Material` asset in the Content Browser.
+  1. Right-click on a `StaticMesh`, `SkeletalMesh`, `AnimSequence`, `Level`, or `Material` asset in the Content Browser.
   1. Select `Asset Actions -> Export...`
   1. Change `Save as type` to `.gltf` and click `Save`
   1. When `Export Options` window is shown, click `Export`
@@ -87,4 +34,4 @@ Usage
   1. Change `Save as type` to `.gltf` and click `Save`
   1. When `Export Options` window is shown, click `Export`
 
-The exported file can be opened in any glTF viewer including those readily available on the web.
+If the Web Viewer is bundled, run `GLTFLaunchHelper.exe` in the target folder to view the export asset effortlessly in the Unreal glTF Viewer.
