[GLTFExporter] chore: minor updates to readme files

commit 449a0b355c1fd2d2177a1b543b788e04f61d81b6
Date: Tue, 1 Dec 2020 04:44:45 +0100

    chore: minor updates to readme files
    (cherry picked from commit d4967578c980187ae174a62ddf3ea488d8c65549​)

#skipci

diff --git a/Engine/Plugins/Enterprise/GLTFExporter/MigratedFromGit/README.md b/Engine/Plugins/Enterprise/GLTFExporter/MigratedFromGit/README.md
--- a/Engine/Plugins/Enterprise/GLTFExporter/MigratedFromGit/README.md
+++ b/Engine/Plugins/Enterprise/GLTFExporter/MigratedFromGit/README.md
@@ -9,12 +9,12 @@ Installation
 
 * Alt 1: Install to project:
   1. Go to project folder which contains `[ProjectName].uproject`
-  1. Create a folder called `Plugins` (if it doesn't already exist).
-  1. Copy `GLTFExporter` folder into the `Plugins` folder.
+  1. Create a folder called `Plugins` (if it doesn't already exist)
+  1. Copy this `GLTFExporter` folder into the `Plugins` folder
 
 * Alt 2: Install to Unreal Engine:
   1. Go to the plugin folder of Unreal Engine which is `Engine/Plugins`
-  1. Copy `GLTFExporter` folder into the `Plugins` folder.
+  1. Copy this `GLTFExporter` folder into the `Plugins` folder
 
 In Unreal Editor, open `Menu -> Edit -> Plugins` and make sure `glTF Exporter` is installed and enabled.
 
@@ -25,13 +25,13 @@ Usage
 * Alt 1: Export asset via Content Browser
   1. Right-click on a `StaticMesh`, `SkeletalMesh`, `AnimSequence`, `Level`, or `Material` asset in the Content Browser.
   1. Select `Asset Actions -> Export...`
-  1. Change `Save as type` to `.gltf` and click `Save`
-  1. When `Export Options` window is shown, click `Export`
+  1. Change `Save as type` to `.gltf` (or `.glb`) and click `Save`
+  1. When `glTF Export Options` window is shown, click `Export`
 
 * Alt 2: Export current level via File Menu
   1. Select any number of actors in the current level
   1. In the top menu, select `File -> Export Selected...`
-  1. Change `Save as type` to `.gltf` and click `Save`
-  1. When `Export Options` window is shown, click `Export`
+  1. Change `Save as type` to `.gltf` (or `.glb`) and click `Save`
+  1. When `glTF Export Options` window is shown, click `Export`
 
-If the Web Viewer is bundled, run `GLTFLaunchHelper.exe` in the target folder to view the export asset effortlessly in the Unreal glTF Viewer.
+If the Web Viewer is bundled, run `GLTFLaunchHelper.exe` in the target folder to view the export asset in the Unreal glTF Viewer.
