[GLTFExporter] chore: remove the installation section as it does not make sense in the context of only sharing plugin on the marketplace

commit 14df5a5107e894f3dbd8d297e8c0b79f4c844222
Date: Mon, 19 Apr 2021 19:34:00 +0100

    chore: remove the installation section as it does not make sense in the context of only sharing plugin on the marketplace
    (cherry picked from commit 2cd8bfbd38a81e86e4ca50fa7d575cd381c5abf5​)

#skipci

diff --git a/Engine/Plugins/Enterprise/GLTFExporter/MigratedFromGit/README.md b/Engine/Plugins/Enterprise/GLTFExporter/MigratedFromGit/README.md
--- a/Engine/Plugins/Enterprise/GLTFExporter/MigratedFromGit/README.md
+++ b/Engine/Plugins/Enterprise/GLTFExporter/MigratedFromGit/README.md
@@ -8,20 +8,6 @@ The glTF Exporter lets UE4 content creators showcase any scene or model on the w
 Currently, the plugin works with both Unreal Engine **4.25** and **4.26**. It has primarily been tested on Windows but is designed to be platform-independent.
 
 
-## Installation
-
-- Alt 1: Install to project:
-  1. Go to project folder which contains `[ProjectName].uproject`
-  1. Create a folder called `Plugins` (if it doesn't already exist)
-  1. Copy this `GLTFExporter` folder into the `Plugins` folder
-
-- Alt 2: Install to Unreal Engine:
-  1. Go to the plugin folder of Unreal Engine which is `Engine/Plugins`
-  1. Copy this `GLTFExporter` folder into the `Plugins` folder
-
-In Unreal Editor, open `Menu -> Edit -> Plugins` and make sure `glTF Exporter` is installed and enabled.
-
-
 ## Usage
 
 - Alt 1: Export asset via Content Browser
