[GLTFExporter] chore: update readmes to include UE 4.27 support

commit b4f15135fb429346c459ab462236d2fcdf5e6648
Date: Fri, 25 Jun 2021 23:29:55 +0200

    chore: update readmes to include UE 4.27 support
    (cherry picked from commit 15777fdd79410c571b9b5fa2a2b28e54724067b3​)

#skipci

diff --git a/Engine/Plugins/Enterprise/GLTFExporter/MigratedFromGit/README.md b/Engine/Plugins/Enterprise/GLTFExporter/MigratedFromGit/README.md
--- a/Engine/Plugins/Enterprise/GLTFExporter/MigratedFromGit/README.md
+++ b/Engine/Plugins/Enterprise/GLTFExporter/MigratedFromGit/README.md
@@ -5,7 +5,7 @@ The glTF Exporter lets UE4 content creators showcase any scene or model on the w
 
 ## Compatibility
 
-Currently, the plugin works with both Unreal Engine **4.25** and **4.26**. It has primarily been tested on Windows but is designed to be platform-independent.
+Currently, the plugin works with Unreal Engine **4.25**, **4.26**, and **4.27**. It has primarily been tested on Windows but is designed to be platform-independent.
 
 
 ## Usage
