[GLTFExporter] chore: improve instructions in readme

commit cfa0b54bc1695bc144b59a03a09b2009041ddb67
Date: Wed, 8 Sep 2021 01:10:52 +0200

    chore: improve instructions in readme
    (cherry picked from commit c78bb85d41aeea7b0603e0b6826475c6cb6e92dd​)

#skipci

diff --git a/Engine/Plugins/Enterprise/GLTFExporter/MigratedFromGit/README.md b/Engine/Plugins/Enterprise/GLTFExporter/MigratedFromGit/README.md
--- a/Engine/Plugins/Enterprise/GLTFExporter/MigratedFromGit/README.md
+++ b/Engine/Plugins/Enterprise/GLTFExporter/MigratedFromGit/README.md
@@ -11,16 +11,15 @@ Currently, the plugin works with Unreal Engine **4.25**, **4.26**, and **4.27**.
 ## Usage
 
 - Alt 1: Export asset via Content Browser
-  1. Right-click on a `StaticMesh`, `SkeletalMesh`, `AnimSequence`, `Level`, or `Material` asset in the Content Browser.
-  1. Select `Asset Actions -> Export...`
-  1. Change `Save as type` to `.gltf` (or `.glb`) and click `Save`
-  1. When `glTF Export Options` window is shown, click `Export`
-
+  - Right-click on a `StaticMesh`, `SkeletalMesh`, `AnimSequence`, `Level`, or `Material` asset in the Content Browser.
+  - Select `Asset Actions -> Export...`
+  - Change `Save as type` to `.gltf` (or `.glb`) and click `Save`
+  - When `glTF Export Options` window is shown, click `Export`
 - Alt 2: Export current level via File Menu
-  1. Select any number of actors in the current level
-  1. In the top menu, select `File -> Export Selected...`
-  1. Change `Save as type` to `.gltf` (or `.glb`) and click `Save`
-  1. When `glTF Export Options` window is shown, click `Export`
+  - Select any number of actors in the current level
+  - In the top menu, select `File -> Export Selected...`
+  - Change `Save as type` to `.gltf` (or `.glb`) and click `Save`
+  - When `glTF Export Options` window is shown, click `Export`
 
 
 ## Documentation
